- Fixing uninitialized variables in remote control and FreeD plugin

diff --git a/Engine/Plugins/VirtualProduction/LiveLinkFreeD/Source/LiveLinkFreeD/Public/LiveLinkFreeDSourceSettings.h b/Engine/Plugins/VirtualProduction/LiveLinkFreeD/Source/LiveLinkFreeD/Public/LiveLinkFreeDSourceSettings.h
--- a/Engine/Plugins/VirtualProduction/LiveLinkFreeD/Source/LiveLinkFreeD/Public/LiveLinkFreeDSourceSettings.h
+++ b/Engine/Plugins/VirtualProduction/LiveLinkFreeD/Source/LiveLinkFreeD/Public/LiveLinkFreeDSourceSettings.h
@@ -10,30 +10,30 @@ USTRUCT()
 struct FFreeDEncoderData
 {
 	GENERATED_BODY()
-
+	
 	/** Is this encoder data valid? */
 	UPROPERTY(EditAnywhere, Category = "Encoder Data")
-	bool bIsValid;
+	bool bIsValid = false;
 
 	/** Invert the encoder input direction */
 	UPROPERTY(EditAnywhere, Category = "Encoder Data", meta = (EditCondition = "bIsValid"))
-	bool bInvertEncoder;
+	bool bInvertEncoder = false;
 
 	/** Use manual Min/Max values for the encoder normalization (normally uses dynamic auto ranging based on inputs) */
 	UPROPERTY(EditAnywhere, Category = "Encoder Data", meta = (EditCondition = "bIsValid"))
-	bool bUseManualRange;
+	bool bUseManualRange = false;
 
 	/** Minimum raw encoder value */
 	UPROPERTY(EditAnywhere, Category = "Encoder Data", meta = (ClampMin = 0, ClampMax = 0x00ffffff, EditCondition = "bIsValid && bUseManualRange"))
-	int32 Min;
+	int32 Min = 0x00FFFFFF;
 
 	/** Maximum raw encoder value */
 	UPROPERTY(EditAnywhere, Category = "Encoder Data", meta = (ClampMin = 0, ClampMax = 0x00ffffff, EditCondition = "bIsValid && bUseManualRange"))
-	int32 Max;
+	int32 Max = 0;
 
 	/** Mask bits for raw encoder value */
 	UPROPERTY(EditAnywhere, Category = "Encoder Data", meta = (ClampMin = 0, ClampMax = 0x00ffffff, EditCondition = "bIsValid"))
-	int32 MaskBits;
+	int32 MaskBits = 0x00FFFFFF;
 };
 
 UENUM(BlueprintType)
