LiveLink FreeD: Update default encoder values
#rnx

diff --git a/Engine/Plugins/VirtualProduction/LiveLinkFreeD/Source/LiveLinkFreeD/Private/LiveLinkFreeDSource.cpp b/Engine/Plugins/VirtualProduction/LiveLinkFreeD/Source/LiveLinkFreeD/Private/LiveLinkFreeDSource.cpp
--- a/Engine/Plugins/VirtualProduction/LiveLinkFreeD/Source/LiveLinkFreeD/Private/LiveLinkFreeDSource.cpp
+++ b/Engine/Plugins/VirtualProduction/LiveLinkFreeD/Source/LiveLinkFreeD/Private/LiveLinkFreeDSource.cpp
@@ -119,7 +119,7 @@ bool FLiveLinkFreeDSource::RequestSourceShutdown()
 // Name			UDP port	Zoom (wide)	(tele)		Focus (near)	(far)		Spare
 // Generic		40000		0x0			0x10000		0x0				0x10000		Unused
 // Panasonic	1111		0x555		0xfff		0x555			0xfff		Iris 0x555 (close) - 0xfff (open)
-// Sony			52380		0x0			0x7ac0		0x1000			0xf000		Lower 12 bits - Iris (F value * 100); Upper 4 bits - Frame number
+// Sony			40000		0x0			0xFFFFFF	0x7FFFFF		0x000000	Lower 12 bits - Iris (F value * 100); Upper 4 bits - Frame number
 // Mosys		8001		0x0			0xffff		0x0				0xffff		Lower 8 bits - Tracking quality 0-3 (undef, good, caution, bad)
 // Stype		6301		0x0			0xffffff	0x0				0xffffff	Unused
 // Ncam			6301		0x0			0xffffff	0x0				0xffffff	Unused
@@ -167,8 +167,8 @@ void FLiveLinkFreeDSource::OnSettingsChanged(ULiveLinkSourceSettings* Settings,
 															SourceSettings->UserDefinedEncoderData = FFreeDEncoderData({ true, true, false, 0x0555, 0x0fff, 0x0000ffff });
 															break;
 
-					case EFreeDDefaultConfigs::Sony:		SourceSettings->FocusDistanceEncoderData = FFreeDEncoderData({ true, false, false, 0x1000, 0xf000, 0x0000ffff });
-															SourceSettings->FocalLengthEncoderData = FFreeDEncoderData({ true, false, false, 0, 0x7ac0, 0x0000ffff });
+					case EFreeDDefaultConfigs::Sony:		SourceSettings->FocusDistanceEncoderData = FFreeDEncoderData({ true, true, false, 0, 0x007fffff, 0x00ffffff });
+															SourceSettings->FocalLengthEncoderData = FFreeDEncoderData({ true, false, false, 0, 0x00ffffff, 0x00ffffff });
 															SourceSettings->UserDefinedEncoderData = FFreeDEncoderData({ true, false, false, 0, 0x0fff, 0x00000fff });
 															break;
 
