LiveLinkFreeD: Add subject name override option.

This covers the multi-camera case where the camera cannot change its camera id.

#rnx

diff --git a/Engine/Plugins/VirtualProduction/LiveLinkFreeD/Source/LiveLinkFreeD/Private/LiveLinkFreeDConnectionSettings.h b/Engine/Plugins/VirtualProduction/LiveLinkFreeD/Source/LiveLinkFreeD/Private/LiveLinkFreeDConnectionSettings.h
--- a/Engine/Plugins/VirtualProduction/LiveLinkFreeD/Source/LiveLinkFreeD/Private/LiveLinkFreeDConnectionSettings.h
+++ b/Engine/Plugins/VirtualProduction/LiveLinkFreeD/Source/LiveLinkFreeD/Private/LiveLinkFreeDConnectionSettings.h
@@ -10,11 +10,15 @@ struct LIVELINKFREED_API FLiveLinkFreeDConnectionSettings
 	GENERATED_BODY()
 
 public:
-	/** Local interface IP address */
-	UPROPERTY(EditAnywhere, Category = "Connection Settings")
+	/** Local interface IP address. 0.0.0.0 binds to the default interface */
+	UPROPERTY(EditAnywhere, Category = "Connection Settings", meta = (DisplayName = "Local IP Address"))
 	FString IPAddress = TEXT("0.0.0.0");
 
 	/** UDP port number */
-	UPROPERTY(EditAnywhere, Category = "Connection Settings")
+	UPROPERTY(EditAnywhere, Category = "Connection Settings", meta = (DisplayName = "UDP Port Number"))
 	uint16 UDPPortNumber = 40000;
+
+	/** Custom subject name (if empty, defaults to "Camera <Camera ID>") */
+	UPROPERTY(EditAnywhere, Category = "Connection Settings", meta = (DisplayName = "Subject Name"))
+	FString SubjectName = TEXT("");
 };
diff --git a/Engine/Plugins/VirtualProduction/LiveLinkFreeD/Source/LiveLinkFreeD/Private/LiveLinkFreeDSource.cpp b/Engine/Plugins/VirtualProduction/LiveLinkFreeD/Source/LiveLinkFreeD/Private/LiveLinkFreeDSource.cpp
--- a/Engine/Plugins/VirtualProduction/LiveLinkFreeD/Source/LiveLinkFreeD/Private/LiveLinkFreeDSource.cpp
+++ b/Engine/Plugins/VirtualProduction/LiveLinkFreeD/Source/LiveLinkFreeD/Private/LiveLinkFreeDSource.cpp
@@ -34,6 +34,7 @@ FLiveLinkFreeDSource::FLiveLinkFreeDSource(const FLiveLinkFreeDConnectionSetting
 : Client(nullptr)
 , Stopping(false)
 , Thread(nullptr)
+, CustomSubjectName(ConnectionSettings.SubjectName)
 {
 	SourceStatus = LOCTEXT("SourceStatus_NoData", "No data");
 	SourceType = LOCTEXT("SourceType_FreeD", "FreeD");
@@ -341,7 +342,15 @@ uint32 FLiveLinkFreeDSource::Run()
 								CameraFrameData->Aperture = UserDefinedData;
 							}
 
-							CameraSubjectName = FString::Printf(TEXT("Camera %d"), CameraId);
+							if (CustomSubjectName.IsEmpty())
+							{
+								CameraSubjectName = FString::Printf(TEXT("Camera %d"), CameraId);
+							}
+							else
+							{
+								CameraSubjectName = CustomSubjectName;
+							}
+
 							Send(&FrameData, FName(CameraSubjectName));
 
 							FrameCounter++;
diff --git a/Engine/Plugins/VirtualProduction/LiveLinkFreeD/Source/LiveLinkFreeD/Public/LiveLinkFreeDSource.h b/Engine/Plugins/VirtualProduction/LiveLinkFreeD/Source/LiveLinkFreeD/Public/LiveLinkFreeDSource.h
--- a/Engine/Plugins/VirtualProduction/LiveLinkFreeD/Source/LiveLinkFreeD/Public/LiveLinkFreeDSource.h
+++ b/Engine/Plugins/VirtualProduction/LiveLinkFreeD/Source/LiveLinkFreeD/Public/LiveLinkFreeDSource.h
@@ -112,6 +112,9 @@ private:
 	// LiveLink subject name for this camera data
 	FString CameraSubjectName = TEXT("Camera");
 
+	// Custom subject name from connection settings
+	FString CustomSubjectName;
+
 	// Pointer to the settings for this source so we don't have to duplicate data
 	ULiveLinkFreeDSourceSettings* SavedSourceSettings = nullptr;
 };
