LiveLinkFreeD: Fix ip address tooltip.

#rnx

diff --git a/Engine/Plugins/VirtualProduction/LiveLinkFreeD/Source/LiveLinkFreeD/Private/LiveLinkFreeDConnectionSettings.h b/Engine/Plugins/VirtualProduction/LiveLinkFreeD/Source/LiveLinkFreeD/Private/LiveLinkFreeDConnectionSettings.h
--- a/Engine/Plugins/VirtualProduction/LiveLinkFreeD/Source/LiveLinkFreeD/Private/LiveLinkFreeDConnectionSettings.h
+++ b/Engine/Plugins/VirtualProduction/LiveLinkFreeD/Source/LiveLinkFreeD/Private/LiveLinkFreeDConnectionSettings.h
@@ -10,7 +10,11 @@ struct LIVELINKFREED_API FLiveLinkFreeDConnectionSettings
 	GENERATED_BODY()
 
 public:
-	/** Local interface IP address. 0.0.0.0 binds to the default interface */
+
+	/**
+	 * Local interface IP address.
+	 * 0.0.0.0 (INADDR_ANY) binds to all local network interfaces.
+	 */
 	UPROPERTY(EditAnywhere, Category = "Connection Settings", meta = (DisplayName = "Local IP Address"))
 	FString IPAddress = TEXT("0.0.0.0");
 
